#UE4 Getting rid of environmental variables for UBT configuration.

diff --git a/Engine/Programs/UnrealBuildTool/BuildConfiguration.xml b/Engine/Programs/UnrealBuildTool/BuildConfiguration.xml
--- a/Engine/Programs/UnrealBuildTool/BuildConfiguration.xml
+++ b/Engine/Programs/UnrealBuildTool/BuildConfiguration.xml
@@ -44,52 +44,80 @@
 	priority. Not defined classes and fields are left alone.
 	-->
 	
-	<!--
-		<BuildConfiguration>
-			<bUseUnityBuild>True</bUseUnityBuild>
-			<bForceUnityBuild>False</bForceUnityBuild>
-			<MinGameModuleSourceFilesForUnityBuild>32</MinGameModuleSourceFilesForUnityBuild>
-			<bUseMonolithicGraphicsDrivers>True</bUseMonolithicGraphicsDrivers>
-			<bUseFastMonoCalls>True</bUseFastMonoCalls>
-			<NumIncludedBytesPerUnityCPP>262144</NumIncludedBytesPerUnityCPP>
-			<bStressTestUnity>False</bStressTestUnity>
-			<bCheckSystemHeadersForModification>False</bCheckSystemHeadersForModification>
-			<bCheckExternalHeadersForModification>False</bCheckExternalHeadersForModification>
-			<bIgnoreOutdatedImportLibraries>True</bIgnoreOutdatedImportLibraries>
-			<bDisableDebugInfo>False</bDisableDebugInfo>
-			<bOmitPCDebugInfoInDevelopment>False</bOmitPCDebugInfoInDevelopment>
-			<bUsePDBFiles>False</bUsePDBFiles>
-			<bUsePCHFiles>True</bUsePCHFiles>
-			<bUseActionHistory>True</bUseActionHistory>
-			<MinFilesUsingPrecompiledHeader>6</MinFilesUsingPrecompiledHeader>
-			<bForcePrecompiledHeaderForGameModules>True</bForcePrecompiledHeaderForGameModules>
-			<bPrintDebugInfo>False</bPrintDebugInfo>
-			<bPrintPerformanceInfo>False</bPrintPerformanceInfo>
-			<bLogDetailedActionStats>False</bLogDetailedActionStats>
-			<bDeployAfterCompile>False</bDeployAfterCompile>
-			<bFlushBuildDirOnRemoteMac>False</bFlushBuildDirOnRemoteMac>
-			<bAllowXGE>True</bAllowXGE>
-			<bXGEExport>False</bXGEExport>
-			<bShowXGEMonitor>False</bShowXGEMonitor>
-			<bStopXGECompilationAfterErrors>True</bStopXGECompilationAfterErrors>
-			<bAllowRemotelyCompiledPCHs>False</bAllowRemotelyCompiledPCHs>
-			<bShouldDeleteAllOutdatedProducedItems>False</bShouldDeleteAllOutdatedProducedItems>
-			<bUseIncrementalLinking>False</bUseIncrementalLinking>
-			<bAllowLTCG>False</bAllowLTCG>
-			<bSupportEditAndContinue>False</bSupportEditAndContinue>
-			<bOmitFramePointers>True</bOmitFramePointers>
-			<ProcessorCountMultiplier>1</ProcessorCountMultiplier>
-			<BaseIntermediateFolder>Intermediate/Build/</BaseIntermediateFolder>
-			<bGeneratedSYMFile>False</bGeneratedSYMFile>
-			<bStripSymbolsOnIOS>False</bStripSymbolsOnIOS>
-			<bCreateStubIPA>True</bCreateStubIPA>
-			<bCopyAppBundleBackToDevice>False</bCopyAppBundleBackToDevice>
-			<bUseMallocProfiler>False</bUseMallocProfiler>
-			<bEnableCodeAnalysis>False</bEnableCodeAnalysis>
-			<bUseSharedPCHs>True</bUseSharedPCHs>
-			<bUseIncludeDependencyResolveCache>True</bUseIncludeDependencyResolveCache>
-			<bTestIncludeDependencyResolveCache>False</bTestIncludeDependencyResolveCache>
-			<bDebugBuildsActuallyUseDebugCRT>False</bDebugBuildsActuallyUseDebugCRT>
-		</BuildConfiguration>
-	-->
+	<BuildConfiguration>
+		<bUseUnityBuild>True</bUseUnityBuild>
+		<bForceUnityBuild>False</bForceUnityBuild>
+		<MinGameModuleSourceFilesForUnityBuild>32</MinGameModuleSourceFilesForUnityBuild>
+		<bUseMonolithicGraphicsDrivers>True</bUseMonolithicGraphicsDrivers>
+		<bUseFastMonoCalls>True</bUseFastMonoCalls>
+		<NumIncludedBytesPerUnityCPP>262144</NumIncludedBytesPerUnityCPP>
+		<bStressTestUnity>False</bStressTestUnity>
+		<bCheckSystemHeadersForModification>False</bCheckSystemHeadersForModification>
+		<bCheckExternalHeadersForModification>False</bCheckExternalHeadersForModification>
+		<bIgnoreOutdatedImportLibraries>True</bIgnoreOutdatedImportLibraries>
+		<bDisableDebugInfo>False</bDisableDebugInfo>
+		<bOmitPCDebugInfoInDevelopment>False</bOmitPCDebugInfoInDevelopment>
+		<bUsePDBFiles>False</bUsePDBFiles>
+		<bUsePCHFiles>True</bUsePCHFiles>
+		<bUseActionHistory>True</bUseActionHistory>
+		<MinFilesUsingPrecompiledHeader>6</MinFilesUsingPrecompiledHeader>
+		<bForcePrecompiledHeaderForGameModules>True</bForcePrecompiledHeaderForGameModules>
+		<bPrintDebugInfo>False</bPrintDebugInfo>
+		<bPrintPerformanceInfo>False</bPrintPerformanceInfo>
+		<bLogDetailedActionStats>False</bLogDetailedActionStats>
+		<bDeployAfterCompile>False</bDeployAfterCompile>
+		<bFlushBuildDirOnRemoteMac>False</bFlushBuildDirOnRemoteMac>
+		<bAllowXGE>True</bAllowXGE>
+		<bXGEExport>False</bXGEExport>
+		<bShowXGEMonitor>False</bShowXGEMonitor>
+		<bStopXGECompilationAfterErrors>True</bStopXGECompilationAfterErrors>
+		<bAllowRemotelyCompiledPCHs>False</bAllowRemotelyCompiledPCHs>
+		<bShouldDeleteAllOutdatedProducedItems>False</bShouldDeleteAllOutdatedProducedItems>
+		<bUseIncrementalLinking>False</bUseIncrementalLinking>
+		<bAllowLTCG>False</bAllowLTCG>
+		<bSupportEditAndContinue>False</bSupportEditAndContinue>
+		<bOmitFramePointers>True</bOmitFramePointers>
+		<ProcessorCountMultiplier>1</ProcessorCountMultiplier>
+		<BaseIntermediateFolder>Intermediate/Build/</BaseIntermediateFolder>
+		<bGeneratedSYMFile>False</bGeneratedSYMFile>
+		<bStripSymbolsOnIOS>False</bStripSymbolsOnIOS>
+		<bCreateStubIPA>True</bCreateStubIPA>
+		<bCopyAppBundleBackToDevice>False</bCopyAppBundleBackToDevice>
+		<bUseMallocProfiler>False</bUseMallocProfiler>
+		<bEnableCodeAnalysis>False</bEnableCodeAnalysis>
+		<bUseSharedPCHs>True</bUseSharedPCHs>
+		<bUseIncludeDependencyResolveCache>True</bUseIncludeDependencyResolveCache>
+		<bTestIncludeDependencyResolveCache>False</bTestIncludeDependencyResolveCache>
+		<bDebugBuildsActuallyUseDebugCRT>False</bDebugBuildsActuallyUseDebugCRT>
+	</BuildConfiguration>
+	<UEBuildConfiguration>
+		<bCompileNetworkProfiler>False</bCompileNetworkProfiler>
+		<bCompileSimplygon>True</bCompileSimplygon>
+		<bCompileRecast>True</bCompileRecast>
+		<bCompileSpeedTree>True</bCompileSpeedTree>
+		<bCompileSteamOSS>True</bCompileSteamOSS>
+		<bCompileMcpOSS>True</bCompileMcpOSS>
+	</UEBuildConfiguration>
+	<AndroidPlatform>
+		<AndroidNdkApiTarget>latest</AndroidNdkApiTarget>
+		<AndroidSdkApiTarget>latest</AndroidSdkApiTarget>
+		<AndroidArchitecture>-armv7</AndroidArchitecture>
+	</AndroidPlatform>
+	<IOSPlatform>
+		<IOSArchitecture></IOSArchitecture>
+	</IOSPlatform>
+	<IOSToolChain>
+		<bUseDangerouslyFastMode>False</bUseDangerouslyFastMode>
+		<IOSVersion>6.0</IOSVersion>
+		<IOSSDKVersion>latest</IOSSDKVersion>
+		<NonShippingArchitectures>armv7</NonShippingArchitectures>
+		<ShippingArchitectures>armv7</ShippingArchitectures>
+		<XcodeDeveloperDir>/Applications/Xcode.app/Contents/Developer/</XcodeDeveloperDir>
+	</IOSToolChain>
+	<HTML5Platform>
+		<HTML5Architecture></HTML5Architecture>
+	</HTML5Platform>
+	<WinRTPlatform>
+		<bCompileWinRT>False</bCompileWinRT>
+	</WinRTPlatform>
 </Configuration>
diff --git a/Engine/Source/Programs/AutomationTool/Mac/MacPlatform.Automation.cs b/Engine/Source/Programs/AutomationTool/Mac/MacPlatform.Automation.cs
--- a/Engine/Source/Programs/AutomationTool/Mac/MacPlatform.Automation.cs
+++ b/Engine/Source/Programs/AutomationTool/Mac/MacPlatform.Automation.cs
@@ -110,7 +110,7 @@ public class MacPlatform : Platform
 	{
 		// don't use the manifest to set up build products if we are compiling Mac under Windows and we aren't going to copy anything back to the PC
 		bool bIsBuildingRemotely = UnrealBuildTool.ExternalExecution.GetRuntimePlatform() != UnrealTargetPlatform.Mac;
-		bool bUseManifest = !bIsBuildingRemotely || UnrealBuildTool.Utils.GetEnvironmentVariable("ue.bCopyAppBundleBackToDevice", false);
+		bool bUseManifest = !bIsBuildingRemotely || BuildConfiguration.bCopyAppBundleBackToDevice;
 		return bUseManifest;
 	}
     public override string GUBP_GetPlatformFailureEMails(string Branch)
diff --git a/Engine/Source/Programs/AutomationTool/Program.cs b/Engine/Source/Programs/AutomationTool/Program.cs
--- a/Engine/Source/Programs/AutomationTool/Program.cs
+++ b/Engine/Source/Programs/AutomationTool/Program.cs
@@ -25,6 +25,8 @@ namespace AutomationTool
 			LogUtils.InitLogging(CommandLine);
 			Log.WriteLine(TraceEventType.Information, "Running on {0}", HostPlatform.Current.GetType().Name);
 
+			XmlConfigLoader.Init();
+
 			// Log if we're running from the launcher
 			var ExecutingAssemblyLocation = CommandUtils.CombinePaths(Assembly.GetExecutingAssembly().Location);
 			if (String.Compare(ExecutingAssemblyLocation, CommandUtils.CombinePaths(InternalUtils.ExecutingAssemblyLocation), true) != 0)
diff --git a/Engine/Source/Programs/UnrealBuildTool/Android/AndroidToolChain.cs b/Engine/Source/Programs/UnrealBuildTool/Android/AndroidToolChain.cs
--- a/Engine/Source/Programs/UnrealBuildTool/Android/AndroidToolChain.cs
+++ b/Engine/Source/Programs/UnrealBuildTool/Android/AndroidToolChain.cs
@@ -16,7 +16,7 @@ namespace UnrealBuildTool
 		static public string GetNdkApiLevel()
 		{
 			// default to looking on disk for latest API level
-			string Target = Utils.GetStringEnvironmentVariable("ue.AndroidNdkApiTarget", "latest");
+			string Target = AndroidPlatform.AndroidNdkApiTarget;
 
 			if (Target == "latest")
 			{
diff --git a/Engine/Source/Programs/UnrealBuildTool/Android/UEBuildAndroid.cs b/Engine/Source/Programs/UnrealBuildTool/Android/UEBuildAndroid.cs
--- a/Engine/Source/Programs/UnrealBuildTool/Android/UEBuildAndroid.cs
+++ b/Engine/Source/Programs/UnrealBuildTool/Android/UEBuildAndroid.cs
@@ -11,12 +11,18 @@ namespace UnrealBuildTool
 {
 	class AndroidPlatform : UEBuildPlatform
 	{
+		/// <summary>
+		/// Android settings.
+		/// </summary>
+		public static string AndroidArchitecture = "-armv7";
+		public static string AndroidNdkApiTarget = "latest";
+		public static string AndroidSdkApiTarget = "latest";
+
 		// The current architecture - affects everything about how UBT operates on Android
-		private static string ActiveArchitecture = Utils.GetStringEnvironmentVariable("ue.AndroidArchitecture", "-armv7");
 		public override string GetActiveArchitecture()
 		{
 			// by default, use an empty architecture (which is really just a modifer to the platform for some paths/names)
-			return ActiveArchitecture;
+			return AndroidArchitecture;
 		}
 
         public override bool CanUseXGE()
diff --git a/Engine/Source/Programs/UnrealBuildTool/Android/UEDeployAndroid.cs b/Engine/Source/Programs/UnrealBuildTool/Android/UEDeployAndroid.cs
--- a/Engine/Source/Programs/UnrealBuildTool/Android/UEDeployAndroid.cs
+++ b/Engine/Source/Programs/UnrealBuildTool/Android/UEDeployAndroid.cs
@@ -58,7 +58,7 @@ namespace UnrealBuildTool.Android
 		private static string GetSdkApiLevel()
 		{
 			// default to looking on disk for latest API level
-			string Target = Utils.GetStringEnvironmentVariable("ue.AndroidSdkApiTarget", "latest");
+			string Target = AndroidPlatform.AndroidSdkApiTarget;
 
 			// if we want to use whatever version the ndk uses, then use that
 			if (Target == "matchndk")
diff --git a/Engine/Source/Programs/UnrealBuildTool/Configuration/BuildConfiguration.cs b/Engine/Source/Programs/UnrealBuildTool/Configuration/BuildConfiguration.cs
--- a/Engine/Source/Programs/UnrealBuildTool/Configuration/BuildConfiguration.cs
+++ b/Engine/Source/Programs/UnrealBuildTool/Configuration/BuildConfiguration.cs
@@ -303,47 +303,47 @@ namespace UnrealBuildTool
 		/// </summary>
 		public static void Reset()
 		{
-			bAllowLTCG = Utils.GetEnvironmentVariable("ue.bAllowLTCG", false);
-			bAllowRemotelyCompiledPCHs = Utils.GetEnvironmentVariable("ue.bAllowRemotelyCompiledPCHs", false);
-            bAllowXGE = Utils.GetEnvironmentVariable("ue.bAllowXGE", true);
+			bAllowLTCG = false;
+			bAllowRemotelyCompiledPCHs = false;
+            bAllowXGE = true;
 
 			// Don't bother to check external (stable) headers for modification.  It slows down UBT's dependency checking.
-			bCheckExternalHeadersForModification = Utils.GetEnvironmentVariable("ue.bCheckExternalHeadersForModification", false);
-			bCheckSystemHeadersForModification = Utils.GetEnvironmentVariable("ue.bCheckSystemHeadersForModification", false);
+			bCheckExternalHeadersForModification = false;
+			bCheckSystemHeadersForModification = false;
 
-			bCopyAppBundleBackToDevice = Utils.GetEnvironmentVariable("ue.bCopyAppBundleBackToDevice", false);
-			bCreateStubIPA = Utils.GetEnvironmentVariable("ue.iPhone_CreateStubIPA", true);
+			bCopyAppBundleBackToDevice = false;
+			bCreateStubIPA = true;
 			bDeployAfterCompile = false;
-			bDisableDebugInfo = Utils.GetEnvironmentVariable("ue.bDisableDebugInfo", false);
+			bDisableDebugInfo = false;
 			bEnableCodeAnalysis = false;
 			bFlushBuildDirOnRemoteMac = false;
-			bGeneratedSYMFile = Utils.GetEnvironmentVariable("ue.bGeneratedSYMFile", false);
-			bStripSymbolsOnIOS = Utils.GetEnvironmentVariable("ue.bStripSymbolsOnIOS", bGeneratedSYMFile);
+			bGeneratedSYMFile = false;
+			bStripSymbolsOnIOS = bGeneratedSYMFile;
 
 			// By default we don't bother relinking targets if only a dependent .lib has changed, as chances are that
 			// the import library wasn't actually different unless a dependent header file of this target was also changed,
 			// in which case the target would be rebuilt automatically.
-			bIgnoreOutdatedImportLibraries = Utils.GetEnvironmentVariable("ue.bIgnoreOutdatedImportLibraries", true);
-
-			bLogDetailedActionStats = Utils.GetEnvironmentVariable("ue.bLogDetailedActionStats", false);
-			bOmitFramePointers = Utils.GetEnvironmentVariable("ue.bOmitFramePointers", true);
-			bOmitPCDebugInfoInDevelopment = Utils.GetEnvironmentVariable("ue.bOmitPCDebugInfoInDevelopment", false);
-			bPrintDebugInfo = Utils.GetEnvironmentVariable("ue.bPrintDebugInfo", false);
-			bPrintPerformanceInfo = Utils.GetEnvironmentVariable("ue.bPrintPerformanceInfo", false);
-			bShouldDeleteAllOutdatedProducedItems = Utils.GetEnvironmentVariable("ue.bShouldDeleteAllOutdatedProducedItems", false);
-			bShowXGEMonitor = Utils.GetEnvironmentVariable("ue.bShowXGEMonitor", false);
-			bStopXGECompilationAfterErrors = Utils.GetEnvironmentVariable("ue.bStopXGECompilationAfterErrors", true);
-			bStressTestUnity = Utils.GetEnvironmentVariable("ue.bStressTestUnity", false);
-			bSupportEditAndContinue = Utils.GetEnvironmentVariable("ue.bSupportEditAndContinue", false);
+			bIgnoreOutdatedImportLibraries = true;
+
+			bLogDetailedActionStats = false;
+			bOmitFramePointers = true;
+			bOmitPCDebugInfoInDevelopment = false;
+			bPrintDebugInfo = false;
+			bPrintPerformanceInfo = false;
+			bShouldDeleteAllOutdatedProducedItems = false;
+			bShowXGEMonitor = false;
+			bStopXGECompilationAfterErrors = true;
+			bStressTestUnity = false;
+			bSupportEditAndContinue = false;
 			bUseActionHistory = true;
 
 			// Incremental linking can yield faster iteration times when making small changes
 			// NOTE: We currently don't use incremental linking because it tends to behave a bit buggy on some computers (PDB-related compile errors)
-			bUseIncrementalLinking = Utils.GetEnvironmentVariable("ue.bUseIncrementalLinking", false);
+			bUseIncrementalLinking = false;
 
 			bUseMallocProfiler = false;
-			bUsePDBFiles = Utils.GetEnvironmentVariable("ue.bUsePDBFiles", false);
-			bUsePCHFiles = Utils.GetEnvironmentVariable("ue.bUsePCHFiles", true);
+			bUsePDBFiles = false;
+			bUsePCHFiles = true;
 
 			// Shared PCHs are necessary to avoid generating tons of private PCH files for every single little module.  While
 			// these private PCHs do yield fastest incremental compiles for single files, it causes full rebuilds to
@@ -355,7 +355,7 @@ namespace UnrealBuildTool
 			// Using unity build to coalesce source files allows for much faster full rebuild times.  For fastest iteration times on single-file
 			// changes, consider turning this off.  In some cases, unity builds can cause linker errors after iterative changes to files, because
 			// the synthesized file names for unity code files may change between builds.
-			bUseUnityBuild = Utils.GetEnvironmentVariable("ue.bUseUnityBuild", true);
+			bUseUnityBuild = true;
 
             bForceUnityBuild = false;
 
@@ -380,11 +380,11 @@ namespace UnrealBuildTool
 
 			// When using the local executor (not XGE), run a single action on each CPU core.  Note that you can set this to a larger value
 			// to get slightly faster build times in many cases, but your computer's responsiveness during compiling may be much worse.
-			ProcessorCountMultiplier = Utils.GetEnvironmentVariable("ue.ProcessorCountMultiplier", 1.0);
+			ProcessorCountMultiplier = 1.0;
 
 			bTestIncludeDependencyResolveCache = false;
 			// if we are testing the resolve cache, we require UBT to use it.
-			bUseIncludeDependencyResolveCache = bTestIncludeDependencyResolveCache || Utils.GetEnvironmentVariable("ue.bUseIncludeDependencyResolveCache", true);
+			bUseIncludeDependencyResolveCache = true;
             
             //IMPORTANT THIS IS THE MAIN SWITCH FOR MONO.
             // New Monolithic Graphics drivers replace D3D and are *mostly* API compatible
diff --git a/Engine/Source/Programs/UnrealBuildTool/Configuration/UEBuildConfiguration.cs b/Engine/Source/Programs/UnrealBuildTool/Configuration/UEBuildConfiguration.cs
--- a/Engine/Source/Programs/UnrealBuildTool/Configuration/UEBuildConfiguration.cs
+++ b/Engine/Source/Programs/UnrealBuildTool/Configuration/UEBuildConfiguration.cs
@@ -114,9 +114,9 @@ namespace UnrealBuildTool
 		{
 			//@todo. Allow disabling PhysX/APEX via these values...
 			// Currently, WITH_PHYSX is forced to true in Engine.h (as it isn't defined anywhere by the builder)
-			bCompilePhysX = true;// Utils.GetEnvironmentVariable("ue.bCompilePhysX", true);
-			bCompileAPEX = true;// Utils.GetEnvironmentVariable("ue.bCompileAPEX", true);
-			bCompileNetworkProfiler = Utils.GetEnvironmentVariable("ue.bCompileNetworkProfiler", false);
+			bCompilePhysX = true;
+			bCompileAPEX = true;
+			bCompileNetworkProfiler = false;
 			bBuildDedicatedServer = false;
 			bBuildEditor = true;
 			bBuildRequiresCookedData = false;
@@ -124,32 +124,42 @@ namespace UnrealBuildTool
 			bBuildDeveloperTools = true;
 			bForceBuildTargetPlatforms = false;
 			bForceBuildShaderFormats = false;
-            bCompileSimplygon = Utils.GetEnvironmentVariable("ue.bCompileSimplygon", true)
-                && Directory.Exists(UEBuildConfiguration.UEThirdPartyDirectory + "NoRedist") == true
-                && Directory.Exists(UEBuildConfiguration.UEThirdPartyDirectory + "NoRedist/Simplygon") == true
-                && Directory.Exists("Developer/SimplygonMeshReduction") == true
-                && !(ProjectFileGenerator.bGenerateProjectFiles && ProjectFileGenerator.bGeneratingRocketProjectFiles);
+			bCompileSimplygon = true;
 			bCompileLeanAndMeanUE = false;
 			bCompileAgainstEngine = true;
 			bCompileAgainstCoreUObject = true;
 			UEThirdPartyDirectory = "ThirdParty/";
-			bCompileRecast = Utils.GetEnvironmentVariable("ue.bCompileRecast", true);
+			bCompileRecast = true;
 			bForceEnableExceptions = false;
 			bWithServerCode = true;
-			bCompileSpeedTree = Utils.GetEnvironmentVariable("ue.bCompileSpeedTree", true)
-				&& Directory.Exists(UEBuildConfiguration.UEThirdPartyDirectory + "NoRedist") == true
-				&& Directory.Exists(UEBuildConfiguration.UEThirdPartyDirectory + "NoRedist/SpeedTree") == true
-				&& !(ProjectFileGenerator.bGenerateProjectFiles && ProjectFileGenerator.bGeneratingRocketProjectFiles);
+			bCompileSpeedTree = true;
 			bCompileWithStatsWithoutEngine = false;
 			bCompileWithPluginSupport = false;
             bUseLoggingInShipping = false;
+			bCompileSteamOSS = true;
+			bCompileMcpOSS = true;
 
+			XmlConfigLoader.Load(typeof(UEBuildConfiguration));
+
+			// Configuration overrides.
 			string SteamVersion = "Steamv128";
-			bCompileSteamOSS = Utils.GetEnvironmentVariable("ue.bCompileSteamOSS", true)
+			bCompileSteamOSS = bCompileSteamOSS
 			   && Directory.Exists(UEBuildConfiguration.UEThirdPartyDirectory + "Steamworks/" + SteamVersion) == true;
 
-			bCompileMcpOSS = Utils.GetEnvironmentVariable("ue.bCompileMcpOSS", true)
-			   && Directory.Exists("Runtime/Online/NoRedist/OnlineSubsystemMcp") == true;	
+			bCompileMcpOSS = bCompileMcpOSS
+			   && Directory.Exists("Runtime/Online/NoRedist/OnlineSubsystemMcp") == true;
+
+
+			bCompileSimplygon = bCompileSimplygon
+				&& Directory.Exists(UEBuildConfiguration.UEThirdPartyDirectory + "NoRedist") == true
+				&& Directory.Exists(UEBuildConfiguration.UEThirdPartyDirectory + "NoRedist/Simplygon") == true
+				&& Directory.Exists("Developer/SimplygonMeshReduction") == true
+				&& !(ProjectFileGenerator.bGenerateProjectFiles && ProjectFileGenerator.bGeneratingRocketProjectFiles);
+
+			bCompileSpeedTree = bCompileSpeedTree
+				&& Directory.Exists(UEBuildConfiguration.UEThirdPartyDirectory + "NoRedist") == true
+				&& Directory.Exists(UEBuildConfiguration.UEThirdPartyDirectory + "NoRedist/SpeedTree") == true
+				&& !(ProjectFileGenerator.bGenerateProjectFiles && ProjectFileGenerator.bGeneratingRocketProjectFiles);
 		}
 
 		/**
diff --git a/Engine/Source/Programs/UnrealBuildTool/HTML5/UEBuildHTML5.cs b/Engine/Source/Programs/UnrealBuildTool/HTML5/UEBuildHTML5.cs
--- a/Engine/Source/Programs/UnrealBuildTool/HTML5/UEBuildHTML5.cs
+++ b/Engine/Source/Programs/UnrealBuildTool/HTML5/UEBuildHTML5.cs
@@ -10,12 +10,13 @@ namespace UnrealBuildTool
 {
 	class HTML5Platform : UEBuildPlatform
 	{
+		public static string HTML5Architecture = "";
+
 		// The current architecture - affects everything about how UBT operates on HTML5
-		private static string ActiveArchitecture = Utils.GetStringEnvironmentVariable("ue.HTML5Architecture", "");//"-win32");
 		public override string GetActiveArchitecture()
 		{
 			// by default, use an empty architecture (which is really just a modifier to the platform for some paths/names)
-			return ActiveArchitecture;
+			return HTML5Architecture;
 		}
 
 		// F5 should always try to run the Win32 version
diff --git a/Engine/Source/Programs/UnrealBuildTool/IOS/IOSToolChain.cs b/Engine/Source/Programs/UnrealBuildTool/IOS/IOSToolChain.cs
--- a/Engine/Source/Programs/UnrealBuildTool/IOS/IOSToolChain.cs
+++ b/Engine/Source/Programs/UnrealBuildTool/IOS/IOSToolChain.cs
@@ -30,27 +30,27 @@ namespace UnrealBuildTool
 		// If you are looking for where to change the remote compile server name, look in RemoteToolChain.cs
 
 		/** If this is set, then we don't do any post-compile steps except moving the executable into the proper spot on the Mac */
-		public static bool bUseDangerouslyFastMode = Utils.GetEnvironmentVariable("ue.IOSDangerouslyFastMode", false);
+		public static bool bUseDangerouslyFastMode = false;
 
 		/** Which version of the iOS SDK to target at build time */
-		public static string IOSSDKVersion = Utils.GetStringEnvironmentVariable("ue.IOSSDKVersion", "latest");
+		public static string IOSSDKVersion = "latest";
 
 		/** The architecture(s) to compile */
-		public static string NonShippingArchitectures = Utils.GetStringEnvironmentVariable("ue.iPhone_NonShippingArchitectures", "armv7");//,arm64");
-		public static string ShippingArchitectures = Utils.GetStringEnvironmentVariable("ue.iPhone_ShippingArchitectures", "armv7");//,arm64");
+		public static string NonShippingArchitectures = "armv7";
+		public static string ShippingArchitectures = "armv7";
 
 		// In case the SDK checking fails for some reason, use this version
 		private static string BackupVersion = "6.0";
 
 		/** Which version of the iOS to allow at run time */
-		public static string IOSVersion = Utils.GetStringEnvironmentVariable("ue.MinIOSVersion", "6.0");
+		public static string IOSVersion = "6.0";
 
 		/** Which developer directory to root from */
-		private static string DeveloperDir = Utils.GetStringEnvironmentVariable("ue.XcodeDeveloperDir", "/Applications/Xcode.app/Contents/Developer/");
+		public static string XcodeDeveloperDir = "/Applications/Xcode.app/Contents/Developer/";
 
 		/** Location of the SDKs */
-		private static string BaseSDKDir = DeveloperDir + "Platforms/iPhoneOS.platform/Developer/SDKs";
-		private static string BaseSDKDirSim = DeveloperDir + "Platforms/iPhoneSimulator.platform/Developer/SDKs";
+		private static string BaseSDKDir;
+		private static string BaseSDKDirSim;
 
 		/** Which compiler frontend to use */
 		private static string IOSCompiler = "clang++";
@@ -61,6 +61,18 @@ namespace UnrealBuildTool
 		/** Which library archiver to use */
 		private static string IOSArchiver = "libtool";
 
+		/// <summary>
+		/// Function to call to reset default data.
+		/// </summary>
+		public static void Reset()
+		{
+			XmlConfigLoader.Load(typeof(IOSToolChain));
+
+			/** Location of the SDKs */
+			BaseSDKDir = XcodeDeveloperDir + "Platforms/iPhoneOS.platform/Developer/SDKs";
+			BaseSDKDirSim = XcodeDeveloperDir + "Platforms/iPhoneSimulator.platform/Developer/SDKs";
+		}
+
 		/** Hunt down the latest IOS sdk if desired */
 		public override void SetUpGlobalEnvironment()
 		{
@@ -274,12 +286,12 @@ namespace UnrealBuildTool
 			if (LinkEnvironment.Config.TargetArchitecture == "-simulator")
 			{
 				Result += " -arch i386";
-				Result += " -isysroot " + DeveloperDir + "Platforms/iPhoneSimulator.platform/Developer/SDKs/iPhoneSimulator" + IOSSDKVersion + ".sdk";
+				Result += " -isysroot " + XcodeDeveloperDir + "Platforms/iPhoneSimulator.platform/Developer/SDKs/iPhoneSimulator" + IOSSDKVersion + ".sdk";
 			}
 			else
 			{
 				Result += Result += GetArchitectureArgument(LinkEnvironment.Config.TargetConfiguration, LinkEnvironment.Config.TargetArchitecture);
-				Result += " -isysroot " + DeveloperDir + "Platforms/iPhoneOS.platform/Developer/SDKs/iPhoneOS" + IOSSDKVersion + ".sdk";
+				Result += " -isysroot " + XcodeDeveloperDir + "Platforms/iPhoneOS.platform/Developer/SDKs/iPhoneOS" + IOSSDKVersion + ".sdk";
 			}
 			Result += " -dead_strip";
 			Result += " -miphoneos-version-min=" + IOSVersion;
@@ -460,7 +472,7 @@ namespace UnrealBuildTool
 				// Add the source file path to the command-line.
 				FileArguments += string.Format(" \"{0}\"", ConvertPath(SourceFile.AbsolutePath), false);
 
-				string CompilerPath = DeveloperDir + "Toolchains/XcodeDefault.xctoolchain/usr/bin/" + IOSCompiler;
+				string CompilerPath = XcodeDeveloperDir + "Toolchains/XcodeDefault.xctoolchain/usr/bin/" + IOSCompiler;
 				if (!Utils.IsRunningOnMono && ExternalExecution.GetRuntimePlatform() != UnrealTargetPlatform.Mac)
 				{
 					CompileAction.ActionHandler = new Action.BlockingActionHandler(RPCUtilHelper.RPCActionHandler);
@@ -482,7 +494,7 @@ namespace UnrealBuildTool
 
 		public override FileItem LinkFiles(LinkEnvironment LinkEnvironment, bool bBuildImportLibraryOnly)
 		{
-			string LinkerPath = DeveloperDir + "Toolchains/XcodeDefault.xctoolchain/usr/bin/" + 
+			string LinkerPath = XcodeDeveloperDir + "Toolchains/XcodeDefault.xctoolchain/usr/bin/" + 
 				(LinkEnvironment.Config.bIsBuildingLibrary ? IOSArchiver : IOSLinker);
 
 			// Create an action that invokes the linker.
@@ -631,7 +643,7 @@ namespace UnrealBuildTool
 
 			// note that the source and dest are switched from a copy command
 			GenDebugAction.CommandArguments = string.Format("-c '{0}/usr/bin/dsymutil {1} -o {2}; cd {2}/..; zip -r -y -1 {3}.app.dSYM.zip {3}.app.dSYM'",
-				DeveloperDir,
+				XcodeDeveloperDir,
 				Executable.AbsolutePath,
 				FullDestPathRoot,
 				Path.GetFileName(Executable.AbsolutePath));
diff --git a/Engine/Source/Programs/UnrealBuildTool/IOS/UEBuildIOS.cs b/Engine/Source/Programs/UnrealBuildTool/IOS/UEBuildIOS.cs
--- a/Engine/Source/Programs/UnrealBuildTool/IOS/UEBuildIOS.cs
+++ b/Engine/Source/Programs/UnrealBuildTool/IOS/UEBuildIOS.cs
@@ -11,12 +11,13 @@ namespace UnrealBuildTool
 {
 	class IOSPlatform : UEBuildPlatform
 	{
+		// by default, use an empty architecture (which is really just a modifer to the platform for some paths/names)
+		public static string IOSArchitecture = "";
+
 		// The current architecture - affects everything about how UBT operates on IOS
-		private static string ActiveArchitecture = Utils.GetStringEnvironmentVariable("ue.IOSArchitecture", "");//"-simulator");
 		public override string GetActiveArchitecture()
 		{
-			// by default, use an empty architecture (which is really just a modifer to the platform for some paths/names)
-			return ActiveArchitecture;
+			return IOSArchitecture;
 		}
 
 		/**
diff --git a/Engine/Source/Programs/UnrealBuildTool/WinRT/UEBuildWinRT.cs b/Engine/Source/Programs/UnrealBuildTool/WinRT/UEBuildWinRT.cs
--- a/Engine/Source/Programs/UnrealBuildTool/WinRT/UEBuildWinRT.cs
+++ b/Engine/Source/Programs/UnrealBuildTool/WinRT/UEBuildWinRT.cs
@@ -10,6 +10,11 @@ namespace UnrealBuildTool
 {
 	class WinRTPlatform : UEBuildPlatform
 	{
+		/// <summary>
+		/// Should the app be compiled as WinRT
+		/// </summary>
+		public static bool bCompileWinRT = false;
+
 		public static bool IsVisualStudioInstalled()
 		{
 			string BaseVSToolPath = WindowsPlatform.Compiler == WindowsCompiler.VisualStudio2012 ? Environment.GetEnvironmentVariable("VS110COMNTOOLS") : Environment.GetEnvironmentVariable("VS120COMNTOOLS");
@@ -656,7 +661,7 @@ namespace UnrealBuildTool
 		 */
 		public static bool ShouldCompileWinRT()
 		{
-			return Utils.GetEnvironmentVariable("ue.bCompileWinRT", false);
+			return WinRTPlatform.bCompileWinRT;
 		}
 	}
 }
