#UE4 CIS fix. XGEExport variable was mistakenly replaced.

diff --git a/Engine/Programs/UnrealBuildTool/BuildConfiguration.xml b/Engine/Programs/UnrealBuildTool/BuildConfiguration.xml
--- a/Engine/Programs/UnrealBuildTool/BuildConfiguration.xml
+++ b/Engine/Programs/UnrealBuildTool/BuildConfiguration.xml
@@ -2,9 +2,9 @@
 	<!--
 	#########################################################################
 	#																		#
-	#	This is a template UBT configuration classes fields override.		#
-	#	If you want value to be overridden just uncomment right				#
-	#	element and put in the desired value.								#
+	#	This is an XML with default UBT configuration. If you want to		#
+	#	override it, create the same file in the locations c. or d.			#
+	#	(see below). DONT'T CHANGE CONTENTS OF THIS FILE!					#	
 	#																		#
 	#########################################################################
 	
@@ -45,50 +45,49 @@
 	-->
 	
 	<BuildConfiguration>
-		<bUseUnityBuild>True</bUseUnityBuild>
-		<bForceUnityBuild>False</bForceUnityBuild>
-		<MinGameModuleSourceFilesForUnityBuild>32</MinGameModuleSourceFilesForUnityBuild>
-		<bUseMonolithicGraphicsDrivers>True</bUseMonolithicGraphicsDrivers>
-		<bUseFastMonoCalls>True</bUseFastMonoCalls>
-		<NumIncludedBytesPerUnityCPP>262144</NumIncludedBytesPerUnityCPP>
-		<bStressTestUnity>False</bStressTestUnity>
-		<bCheckSystemHeadersForModification>False</bCheckSystemHeadersForModification>
+		<bAllowLTCG>False</bAllowLTCG>
+		<bAllowRemotelyCompiledPCHs>False</bAllowRemotelyCompiledPCHs>
+		<bAllowXGE>True</bAllowXGE>
 		<bCheckExternalHeadersForModification>False</bCheckExternalHeadersForModification>
-		<bIgnoreOutdatedImportLibraries>True</bIgnoreOutdatedImportLibraries>
+		<bCheckSystemHeadersForModification>False</bCheckSystemHeadersForModification>
+		<bCopyAppBundleBackToDevice>False</bCopyAppBundleBackToDevice>
+		<bCreateStubIPA>True</bCreateStubIPA>
+		<bDeployAfterCompile>False</bDeployAfterCompile>
 		<bDisableDebugInfo>False</bDisableDebugInfo>
+		<bEnableCodeAnalysis>False</bEnableCodeAnalysis>
+		<bFlushBuildDirOnRemoteMac>False</bFlushBuildDirOnRemoteMac>
+		<bGeneratedSYMFile>False</bGeneratedSYMFile>
+		<bStripSymbolsOnIOS>False</bStripSymbolsOnIOS>
+		<bIgnoreOutdatedImportLibraries>True</bIgnoreOutdatedImportLibraries>
+		<bLogDetailedActionStats>False</bLogDetailedActionStats>
+		<bOmitFramePointers>True</bOmitFramePointers>
 		<bOmitPCDebugInfoInDevelopment>False</bOmitPCDebugInfoInDevelopment>
-		<bUsePDBFiles>False</bUsePDBFiles>
-		<bUsePCHFiles>True</bUsePCHFiles>
-		<bUseActionHistory>True</bUseActionHistory>
-		<MinFilesUsingPrecompiledHeader>6</MinFilesUsingPrecompiledHeader>
-		<bForcePrecompiledHeaderForGameModules>True</bForcePrecompiledHeaderForGameModules>
 		<bPrintDebugInfo>False</bPrintDebugInfo>
 		<bPrintPerformanceInfo>False</bPrintPerformanceInfo>
-		<bLogDetailedActionStats>False</bLogDetailedActionStats>
-		<bDeployAfterCompile>False</bDeployAfterCompile>
-		<bFlushBuildDirOnRemoteMac>False</bFlushBuildDirOnRemoteMac>
-		<bAllowXGE>True</bAllowXGE>
-		<bXGEExport>False</bXGEExport>
+		<bShouldDeleteAllOutdatedProducedItems>False</bShouldDeleteAllOutdatedProducedItems>
 		<bShowXGEMonitor>False</bShowXGEMonitor>
 		<bStopXGECompilationAfterErrors>True</bStopXGECompilationAfterErrors>
-		<bAllowRemotelyCompiledPCHs>False</bAllowRemotelyCompiledPCHs>
-		<bShouldDeleteAllOutdatedProducedItems>False</bShouldDeleteAllOutdatedProducedItems>
-		<bUseIncrementalLinking>False</bUseIncrementalLinking>
-		<bAllowLTCG>False</bAllowLTCG>
+		<bStressTestUnity>False</bStressTestUnity>
 		<bSupportEditAndContinue>False</bSupportEditAndContinue>
-		<bOmitFramePointers>True</bOmitFramePointers>
-		<ProcessorCountMultiplier>1</ProcessorCountMultiplier>
-		<BaseIntermediateFolder>Intermediate/Build/</BaseIntermediateFolder>
-		<bGeneratedSYMFile>False</bGeneratedSYMFile>
-		<bStripSymbolsOnIOS>False</bStripSymbolsOnIOS>
-		<bCreateStubIPA>True</bCreateStubIPA>
-		<bCopyAppBundleBackToDevice>False</bCopyAppBundleBackToDevice>
+		<bUseActionHistory>True</bUseActionHistory>
+		<bUseIncrementalLinking>False</bUseIncrementalLinking>
 		<bUseMallocProfiler>False</bUseMallocProfiler>
-		<bEnableCodeAnalysis>False</bEnableCodeAnalysis>
+		<bUsePDBFiles>False</bUsePDBFiles>
+		<bUsePCHFiles>True</bUsePCHFiles>
 		<bUseSharedPCHs>True</bUseSharedPCHs>
-		<bUseIncludeDependencyResolveCache>True</bUseIncludeDependencyResolveCache>
+		<bUseUnityBuild>True</bUseUnityBuild>
+		<bForceUnityBuild>False</bForceUnityBuild>
+		<MinGameModuleSourceFilesForUnityBuild>32</MinGameModuleSourceFilesForUnityBuild>
+		<NumIncludedBytesPerUnityCPP>262144</NumIncludedBytesPerUnityCPP>
+		<MinFilesUsingPrecompiledHeader>6</MinFilesUsingPrecompiledHeader>
+		<bForcePrecompiledHeaderForGameModules>True</bForcePrecompiledHeaderForGameModules>
+		<ProcessorCountMultiplier>1.0</ProcessorCountMultiplier>
 		<bTestIncludeDependencyResolveCache>False</bTestIncludeDependencyResolveCache>
+		<bUseIncludeDependencyResolveCache>True</bUseIncludeDependencyResolveCache>
+		<bUseMonolithicGraphicsDrivers>True</bUseMonolithicGraphicsDrivers>
+		<bUseFastMonoCalls>True</bUseFastMonoCalls>
 		<bDebugBuildsActuallyUseDebugCRT>False</bDebugBuildsActuallyUseDebugCRT>
+		<BaseIntermediateFolder>"Intermediate/Build/"</BaseIntermediateFolder>
 	</BuildConfiguration>
 	<UEBuildConfiguration>
 		<bCompileNetworkProfiler>False</bCompileNetworkProfiler>
