#UE4 CIS fix.

diff --git a/Engine/Programs/UnrealBuildTool/BuildConfiguration.xml b/Engine/Programs/UnrealBuildTool/BuildConfiguration.xml
--- a/Engine/Programs/UnrealBuildTool/BuildConfiguration.xml
+++ b/Engine/Programs/UnrealBuildTool/BuildConfiguration.xml
@@ -87,7 +87,7 @@
 		<bUseMonolithicGraphicsDrivers>True</bUseMonolithicGraphicsDrivers>
 		<bUseFastMonoCalls>True</bUseFastMonoCalls>
 		<bDebugBuildsActuallyUseDebugCRT>False</bDebugBuildsActuallyUseDebugCRT>
-		<BaseIntermediateFolder>"Intermediate/Build/"</BaseIntermediateFolder>
+		<BaseIntermediateFolder>Intermediate/Build/</BaseIntermediateFolder>
 	</BuildConfiguration>
 	<UEBuildConfiguration>
 		<bCompileNetworkProfiler>False</bCompileNetworkProfiler>
