Turn unity build back on by default

diff --git a/Engine/Programs/UnrealBuildTool/BuildConfiguration.xml b/Engine/Programs/UnrealBuildTool/BuildConfiguration.xml
--- a/Engine/Programs/UnrealBuildTool/BuildConfiguration.xml
+++ b/Engine/Programs/UnrealBuildTool/BuildConfiguration.xml
@@ -75,7 +75,7 @@
 		<bUsePDBFiles>False</bUsePDBFiles>
 		<bUsePCHFiles>True</bUsePCHFiles>
 		<bUseSharedPCHs>True</bUseSharedPCHs>
-		<bUseUnityBuild>False</bUseUnityBuild>
+		<bUseUnityBuild>True</bUseUnityBuild>
 		<bForceUnityBuild>False</bForceUnityBuild>
 		<MinGameModuleSourceFilesForUnityBuild>32</MinGameModuleSourceFilesForUnityBuild>
 		<NumIncludedBytesPerUnityCPP>262144</NumIncludedBytesPerUnityCPP>
